Allow IasTool to build for Linux

diff --git a/Engine/Source/Programs/IoStoreOnDemand/IasTool.Target.cs b/Engine/Source/Programs/IoStoreOnDemand/IasTool.Target.cs
--- a/Engine/Source/Programs/IoStoreOnDemand/IasTool.Target.cs
+++ b/Engine/Source/Programs/IoStoreOnDemand/IasTool.Target.cs
@@ -2,7 +2,7 @@
 
 using UnrealBuildTool;
 
-[SupportedPlatforms("Win64")]
+[SupportedPlatforms("Win64", "Linux")]
 public class IasToolTarget : TargetRules
 {
 	public IasToolTarget(TargetInfo Target) : base(Target)
