Minor clean up of the IASTool target

#rnx

- Remove duplicate entry setting bCompileICU
- Add a code comment showing that the trace system will be disabled and leave a commented out line to enable it.
-- I was attempting to take an insights trace of IASTool and was confused for a short while as to why it wasn't working. This should make it clearer in the future.

diff --git a/Engine/Source/Programs/IoStoreOnDemand/IasTool.Target.cs b/Engine/Source/Programs/IoStoreOnDemand/IasTool.Target.cs
--- a/Engine/Source/Programs/IoStoreOnDemand/IasTool.Target.cs
+++ b/Engine/Source/Programs/IoStoreOnDemand/IasTool.Target.cs
@@ -20,7 +20,6 @@ public class IasToolTarget : TargetRules
 		bCompileAgainstCoreUObject = false;
 		bCompileAgainstEngine = false;
 		bCompileICU = false;
-		bCompileICU = false;
 		bCompileNavmeshClusterLinks = false;
 		bCompileNavmeshSegmentLinks = false;
 		bCompileRecast = false;
@@ -30,5 +29,8 @@ public class IasToolTarget : TargetRules
 		bUsesSlate = false;
 		bWithLiveCoding = false;
 		bWithServerCode = false;
+
+		// Programs have the Insights trace system disabled by default. Uncomment the line below to enable it.
+		//bEnableTrace = true;
 	}
 }
