IASTool can now create UnrealInsight traces.

#rnx

#Synced-CL 49815329

diff --git a/Engine/Source/Programs/IoStoreOnDemand/IasTool.Target.cs b/Engine/Source/Programs/IoStoreOnDemand/IasTool.Target.cs
--- a/Engine/Source/Programs/IoStoreOnDemand/IasTool.Target.cs
+++ b/Engine/Source/Programs/IoStoreOnDemand/IasTool.Target.cs
@@ -28,7 +28,6 @@ public class IasToolTarget : TargetRules
 		bWithLiveCoding = false;
 		bWithServerCode = false;
 
-		// Programs have the Insights trace system disabled by default. Uncomment the line below to enable it.
-		//bEnableTrace = true;
+		bEnableTrace = true;
 	}
 }
